move OpenAL Module from platform specific folder to main tree, now that its used by multiple platforms.

diff --git a/Engine/Source/Runtime/ALAudio/ALAudio.Build.cs b/Engine/Source/Runtime/ALAudio/ALAudio.Build.cs
new file mode 100644
--- /dev/null
+++ b/Engine/Source/Runtime/ALAudio/ALAudio.Build.cs
@@ -0,0 +1,26 @@
+// Copyright 1998-2015 Epic Games, Inc. All Rights Reserved.
+
+using UnrealBuildTool;
+
+public class ALAudio : ModuleRules
+{
+    public ALAudio(TargetInfo Target)
+    {
+        PrivateIncludePathModuleNames.Add("TargetPlatform");
+
+        PrivateDependencyModuleNames.AddRange(new string[]
+        {
+			"Core",
+			"CoreUObject",
+			"Engine",
+		});
+
+        AddThirdPartyPrivateStaticDependencies(Target, "OpenAL");
+        AddThirdPartyPrivateStaticDependencies(Target,
+            "OpenAL",
+            "UEOgg",
+            "Vorbis",
+            "VorbisFile"
+        );
+    }
+}
diff --git a/Engine/Source/Runtime/ALAudio/Private/ALAudioBuffer.cpp b/Engine/Source/Runtime/ALAudio/Private/ALAudioBuffer.cpp
new file mode 100644
--- /dev/null
+++ b/Engine/Source/Runtime/ALAudio/Private/ALAudioBuffer.cpp
@@ -0,0 +1,193 @@
+// Copyright 1998-2015 Epic Games, Inc. All Rights Reserved.
+
+#include "ALAudioDevice.h"
+#include "AudioEffect.h"
+#include "Engine.h"
+#include "IAudioFormat.h"
+#include "AudioDecompress.h"
+#include "ContentStreaming.h"
+
+/*------------------------------------------------------------------------------------
+	FALSoundBuffer.
+------------------------------------------------------------------------------------*/
+/**
+ * Constructor
+ *
+ * @param AudioDevice	audio device this sound buffer is going to be attached to.
+ */
+FALSoundBuffer::FALSoundBuffer( FALAudioDevice* InAudioDevice )
+	: FSoundBuffer(InAudioDevice),
+	  BufferId(0)
+{
+}
+
+/**
+ * Destructor
+ *
+ * Frees wave data and detaches itself from audio device.
+ */
+FALSoundBuffer::~FALSoundBuffer( void )
+{
+	// Delete AL buffers.
+	alDeleteBuffers(1, &BufferId);
+}
+
+/**
+ * Static function used to create a buffer.
+ *
+ * @param InWave		USoundNodeWave to use as template and wave source
+ * @param AudioDevice	audio device to attach created buffer to
+ * @param	bIsPrecacheRequest	Whether this request is for precaching or not
+ * @return	FALSoundBuffer pointer if buffer creation succeeded, NULL otherwise
+ */
+
+FALSoundBuffer* FALSoundBuffer::Init(FALAudioDevice* AudioDevice, USoundWave* InWave)
+{
+	// Can't create a buffer without any source data
+	if (InWave == nullptr || InWave->NumChannels == 0)
+	{
+		return nullptr;
+	}
+
+	FAudioDeviceManager* AudioDeviceManager = GEngine->GetAudioDeviceManager();
+	check(AudioDeviceManager != nullptr);
+
+	FALSoundBuffer *Buffer = nullptr;
+
+	switch (static_cast<EDecompressionType>(InWave->DecompressionType))
+	{
+	case DTYPE_Setup:
+		// Has circumvented pre-cache mechanism - pre-cache now
+		AudioDevice->Precache(InWave, true, false);
+
+		// Recall this function with new decompression type
+		return Init(AudioDevice, InWave);
+
+	case DTYPE_Native:
+		if (InWave->ResourceID)
+		{
+			Buffer = static_cast<FALSoundBuffer*>(AudioDeviceManager->WaveBufferMap.FindRef(InWave->ResourceID));
+		}
+
+		if (!Buffer || !Buffer->BufferId)
+		{
+			CreateNativeBuffer(AudioDevice, InWave, Buffer);
+		}
+		break;
+
+		// Always create a new buffer for streaming ogg vorbis data
+		//Buffer = CreateQueuedBuffer(AudioDevice, InWave);
+		//break;
+
+	case DTYPE_Invalid:
+	case DTYPE_Preview:
+	case DTYPE_Procedural:
+	case DTYPE_RealTime:
+	default:
+		// Invalid will be set if the wave cannot be played
+		UE_LOG(LogALAudio, Warning, TEXT("ALSoundBuffer wave '%s' has an invalid decompression type %d."),
+			*InWave->GetName(), static_cast<EDecompressionType>(InWave->DecompressionType));
+		break;
+	}
+
+	if (Buffer == nullptr)
+	{
+		UE_LOG(LogALAudio, Warning, TEXT("ALSoundBuffer init failed for wave '%s', decompression type %d."),
+			*InWave->GetName(), static_cast<EDecompressionType>(InWave->DecompressionType));
+	}
+
+	return Buffer;
+}
+
+void FALSoundBuffer::CreateNativeBuffer(FALAudioDevice* AudioDevice, USoundWave* Wave, FALSoundBuffer* &Buffer)
+{
+	// Check if Buffer already exists
+	if (Buffer)
+	{
+		// Assign the new AudioDevice
+		Buffer->AudioDevice = AudioDevice;
+	}
+	else
+	{
+		SCOPE_CYCLE_COUNTER( STAT_AudioResourceCreationTime );
+
+		// Check to see if thread has finished decompressing on the other thread
+		if (Wave->AudioDecompressor)
+		{
+			Wave->AudioDecompressor->EnsureCompletion();
+
+			// Remove the decompressor
+			delete Wave->AudioDecompressor;
+			Wave->AudioDecompressor = nullptr;
+		}
+
+		Wave->InitAudioResource(AudioDevice->GetRuntimeFormat(Wave));
+
+		// Create new buffer.
+		Buffer = new FALSoundBuffer(AudioDevice);
+
+		Buffer->InternalFormat = AudioDevice->GetInternalFormat(Wave->NumChannels);
+		Buffer->NumChannels = Wave->NumChannels;
+		Buffer->SampleRate = Wave->SampleRate;
+
+		FAudioDeviceManager* AudioDeviceManager = GEngine->GetAudioDeviceManager();
+		check(AudioDeviceManager != nullptr);
+		AudioDeviceManager->TrackResource(Wave, Buffer);
+	}
+
+	// Generate the new OpenAL buffer
+	alGenBuffers(1, &Buffer->BufferId);
+	AudioDevice->alError(TEXT("RegisterSound"));
+
+	if (Wave->RawPCMData)
+	{
+		// upload it
+		Buffer->BufferSize = Wave->RawPCMDataSize;
+		alBufferData(Buffer->BufferId, Buffer->InternalFormat, Wave->RawPCMData, Wave->RawPCMDataSize, Buffer->SampleRate);
+
+		// Free up the data if necessary
+		if (Wave->bDynamicResource)
+		{
+			FMemory::Free( Wave->RawPCMData );
+			Wave->RawPCMData = nullptr;
+			Wave->bDynamicResource = false;
+		}
+	}
+	else
+	{
+		// get the raw data
+		uint8* SoundData = reinterpret_cast<uint8*>(Wave->RawData.Lock(LOCK_READ_ONLY));
+		// it's (possibly) a pointer to a wave file, so skip over the header
+		int SoundDataSize = Wave->RawData.GetBulkDataSize();
+
+		// is there a wave header?
+		FWaveModInfo WaveInfo;
+		if (WaveInfo.ReadWaveInfo(SoundData, SoundDataSize))
+		{
+			// if so, modify the location and size of the sound data based on header
+			SoundData = WaveInfo.SampleDataStart;
+			SoundDataSize = WaveInfo.SampleDataSize;
+		}
+		// let the Buffer know the final size
+		Buffer->BufferSize = SoundDataSize;
+
+		// upload it
+		alBufferData( Buffer->BufferId, Buffer->InternalFormat, SoundData, Buffer->BufferSize, Buffer->SampleRate );
+		// unload it
+		Wave->RawData.Unlock();
+	}
+
+	if (AudioDevice->alError(TEXT( "RegisterSound (buffer data)")) || (Buffer->BufferSize == 0))
+	{
+		Buffer->InternalFormat = 0;
+	}
+
+	if (Buffer->InternalFormat == 0)
+	{
+		UE_LOG(LogAudio, Log,TEXT( "Audio: sound format not supported for '%s' (%d)" ), *Wave->GetName(), Wave->NumChannels);
+		UE_LOG(LogALAudio, Warning, TEXT("ALSoundBuffer: sound format not supported for wave '%s'"), *Wave->GetName());
+
+		delete Buffer;
+		Buffer = nullptr;
+	}
+}
diff --git a/Engine/Source/Runtime/ALAudio/Private/ALAudioDevice.cpp b/Engine/Source/Runtime/ALAudio/Private/ALAudioDevice.cpp
new file mode 100644
--- /dev/null
+++ b/Engine/Source/Runtime/ALAudio/Private/ALAudioDevice.cpp
@@ -0,0 +1,342 @@
+// Copyright 1998-2015 Epic Games, Inc. All Rights Reserved.
+
+/*=============================================================================
+	ALAudioDevice.cpp: Unreal OpenAL Audio interface object.
+
+	Unreal is RHS with Y and Z swapped (or technically LHS with flipped axis)
+=============================================================================*/
+
+/*------------------------------------------------------------------------------------
+	Audio includes.
+------------------------------------------------------------------------------------*/
+
+#include "ALAudioDevice.h"
+#include "VorbisAudioInfo.h"
+#include "AudioEffect.h"
+#include "Engine.h"
+
+DEFINE_LOG_CATEGORY(LogALAudio);
+
+
+class FALAudioDeviceModule : public IAudioDeviceModule
+{
+public:
+
+	/** Creates a new instance of the audio device implemented by the module. */
+	virtual FAudioDevice* CreateAudioDevice() override
+	{
+		return new FALAudioDevice;
+	}
+};
+
+IMPLEMENT_MODULE(FALAudioDeviceModule, ALAudio );
+/*------------------------------------------------------------------------------------
+	UALAudioDevice constructor and UObject interface.
+------------------------------------------------------------------------------------*/
+FALAudioDevice::FALAudioDevice()
+	: HardwareDevice(nullptr)
+	, SoundContext(nullptr)
+{}
+
+void FALAudioDevice::TeardownHardware( void )
+{
+	// Push any pending data to the hardware
+	if( &alcProcessContext )
+	{
+		alcProcessContext( SoundContext );
+	}
+
+	// Destroy all sound sources
+	for( int32 i = 0; i < Sources.Num(); i++ )
+	{
+		delete Sources[ i ];
+	}
+	Sources.Empty();
+	FreeSources.Empty();
+
+	// Destroy OpenAL buffers associated with this audio device
+	FAudioDeviceManager* AudioDeviceManager = GEngine->GetAudioDeviceManager();
+	check(AudioDeviceManager != nullptr);
+	for( int32 i = 0; i < AudioDeviceManager->Buffers.Num(); i++ )
+	{
+		FALSoundBuffer* Buffer = static_cast<FALSoundBuffer*>(AudioDeviceManager->Buffers[i]);
+		if( Buffer->AudioDevice == this )
+		{
+			alDeleteBuffers(1, &Buffer->BufferId);
+			Buffer->AudioDevice = nullptr;
+			Buffer->BufferId = 0;
+		}
+	}
+
+	// Disable the context
+	if( &alcMakeContextCurrent )
+	{
+		alcMakeContextCurrent(nullptr);
+	}
+
+	// Destroy the context
+	if( &alcDestroyContext )
+	{
+		alcDestroyContext( SoundContext );
+		SoundContext = NULL;
+	}
+
+	// Close the hardware device
+	if( &alcCloseDevice )
+	{
+		if (HardwareDevice)
+		{
+			const ALCchar* DeviceName = alcGetString(HardwareDevice, ALC_DEVICE_SPECIFIER);
+			UE_LOG(LogALAudio, Log, TEXT("Closing ALAudio device : %s"), StringCast<TCHAR>(static_cast<const ANSICHAR*>(DeviceName)).Get());
+
+			alcCloseDevice(HardwareDevice);
+			HardwareDevice = nullptr;
+		}
+	}
+
+}
+
+/*------------------------------------------------------------------------------------
+	UAudioDevice Interface.
+------------------------------------------------------------------------------------*/
+/**
+ * Initializes the audio device and creates sources.
+ *
+ * @warning:
+ *
+ * @return TRUE if initialization was successful, FALSE otherwise
+ */
+bool FALAudioDevice::InitializeHardware( void )
+{
+	// Make sure no interface classes contain any garbage
+	Effects = NULL;
+	DLLHandle = NULL;
+
+	// Default to sensible channel count.
+	if( MaxChannels < 1 )
+	{
+		MaxChannels = 32;
+	}
+	// Load ogg and vorbis dlls if they haven't been loaded yet
+	//LoadVorbisLibraries();
+
+	// Open device
+	HardwareDevice = alcOpenDevice(nullptr);
+	if( !HardwareDevice )
+	{
+		UE_LOG(LogALAudio, Log, TEXT( "ALAudio: no OpenAL devices found." ) );
+		return  false  ;
+	}
+
+	// Display the audio device that was actually opened
+	const ALCchar* OpenedDeviceName = alcGetString( HardwareDevice, ALC_DEVICE_SPECIFIER );
+	UE_LOG(LogALAudio, Log, TEXT("ALAudio device opened : %s"), StringCast<TCHAR>(static_cast<const ANSICHAR*>(OpenedDeviceName)).Get());
+
+	// Create a context
+	int Caps[] =
+	{
+		ALC_FREQUENCY, 44100,
+		ALC_STEREO_SOURCES, 4,
+		0, 0
+	};
+#if PLATFORM_HTML5_WIN32 || PLATFORM_LINUX
+	SoundContext = alcCreateContext( HardwareDevice, Caps );
+#elif PLATFORM_HTML5
+	SoundContext = alcCreateContext( HardwareDevice, 0 );
+#endif
+
+	if( !SoundContext )
+	{
+		return false ;
+	}
+
+	alcMakeContextCurrent(SoundContext);
+
+	// Make sure everything happened correctly
+	if( alError( TEXT( "Init" ) ) )
+	{
+		UE_LOG(LogALAudio, Warning, TEXT("ALAudio: alcMakeContextCurrent failed."));
+		return false ;
+	}
+
+	UE_LOG(LogALAudio, Log, TEXT("AL_VENDOR      : %s"), StringCast<TCHAR>(static_cast<const ANSICHAR*>(alGetString(AL_VENDOR))).Get());
+	UE_LOG(LogALAudio, Log, TEXT("AL_RENDERER    : %s"), StringCast<TCHAR>(static_cast<const ANSICHAR*>(alGetString(AL_RENDERER))).Get());
+	UE_LOG(LogALAudio, Log, TEXT("AL_VERSION     : %s"), StringCast<TCHAR>(static_cast<const ANSICHAR*>(alGetString(AL_VERSION))).Get());
+	UE_LOG(LogALAudio, Log, TEXT("AL_EXTENSIONS  : %s"), StringCast<TCHAR>(static_cast<const ANSICHAR*>(alGetString(AL_EXTENSIONS))).Get());
+
+	// Get the enums for multichannel support
+#if !PLATFORM_HTML5
+	Surround40Format = alGetEnumValue( "AL_FORMAT_QUAD16" );
+	Surround51Format = alGetEnumValue( "AL_FORMAT_51CHN16" );
+	Surround61Format = alGetEnumValue( "AL_FORMAT_61CHN16" );
+	Surround71Format = alGetEnumValue( "AL_FORMAT_71CHN16" );
+#endif
+	// Initialize channels.
+	alError( TEXT( "Emptying error stack" ), 0 );
+	for( int i = 0; i < (( MaxChannels >  MAX_AUDIOCHANNELS ) ? MAX_AUDIOCHANNELS : MaxChannels); i++ )
+	{
+		ALuint SourceId;
+		alGenSources( 1, &SourceId );
+		if( !alError( TEXT( "Init (creating sources)" ), 0 ) )
+		{
+			FALSoundSource* Source = new FALSoundSource( this );
+			Source->SourceId = SourceId;
+			Sources.Add( Source );
+			FreeSources.Add( Source );
+		}
+		else
+		{
+			break;
+		}
+	}
+
+	if( Sources.Num() < 1 )
+	{
+		UE_LOG(LogALAudio, Warning, TEXT("ALAudio: couldn't allocate any sources"));
+		return false ;
+	}
+
+	// Update MaxChannels in case we couldn't create enough sources.
+	MaxChannels = Sources.Num();
+	UE_LOG(LogALAudio, Verbose, TEXT("ALAudioDevice: Allocated %i sources"), MaxChannels);
+
+	// Use our own distance model.
+	alDistanceModel( AL_NONE );
+
+	// Set up a default (nop) effects manager
+	Effects = new FAudioEffectsManager( this );
+
+	return true ;
+}
+
+/**
+ * Update the audio device and calculates the cached inverse transform later
+ * on used for spatialization.
+ */
+void FALAudioDevice::UpdateHardware()
+{
+	if (Listeners.Num() > 0)
+	{
+		// Caches the matrix used to transform a sounds position into local space so we can just look
+		// at the Y component after normalization to determine spatialization.
+		const FVector Up = Listeners[0].GetUp();
+		const FVector Right = Listeners[0].GetFront();
+		InverseTransform = FMatrix(Up, Right, Up ^ Right, Listeners[0].Transform.GetTranslation()).InverseFast();
+	}
+}
+
+ALuint FALAudioDevice::GetInternalFormat( int NumChannels )
+{
+	ALuint InternalFormat = 0;
+
+	switch( NumChannels )
+	{
+	case 0:
+	case 3:
+	case 5:
+		break;
+	case 1:
+		InternalFormat = AL_FORMAT_MONO16;
+		break;
+	case 2:
+		InternalFormat = AL_FORMAT_STEREO16;
+		break;
+#if !PLATFORM_HTML5
+	case 4:
+		InternalFormat = Surround40Format;
+		break;
+	case 6:
+		InternalFormat = Surround51Format;
+		break;
+	case 7:
+		InternalFormat = Surround61Format;
+		break;
+	case 8:
+		InternalFormat = Surround71Format;
+		break;
+#endif
+	}
+
+	return( InternalFormat );
+}
+
+void FALAudioDevice::MakeCurrent(const TCHAR * CallSiteIdentifier)
+{
+#if PLATFORM_LINUX
+	checkf(SoundContext, TEXT("Unitialiized sound context in FALAudioDevice::MakeCurrent()!"));
+	if (!alcMakeContextCurrent(SoundContext))
+	{
+		alError(CallSiteIdentifier ? CallSiteIdentifier : TEXT("FALAudioDevice::MakeCurrent()"));
+	}
+#endif
+}
+
+/*------------------------------------------------------------------------------------
+OpenAL utility functions
+------------------------------------------------------------------------------------*/
+
+//
+//	alError
+//
+bool FALAudioDevice::alError( const TCHAR* Text, bool Log )
+{
+	ALenum Error = alGetError();
+	if( Error != AL_NO_ERROR )
+	{
+		do
+		{
+			if( Log )
+			{
+				switch ( Error )
+				{
+					case AL_INVALID_NAME:
+						UE_LOG(LogALAudio, Warning, TEXT("ALAudio: AL_INVALID_NAME in %s"), Text);
+						break;
+					case AL_INVALID_VALUE:
+						UE_LOG(LogALAudio, Warning, TEXT("ALAudio: AL_INVALID_VALUE in %s"), Text);
+						break;
+					case AL_OUT_OF_MEMORY:
+						UE_LOG(LogALAudio, Warning, TEXT("ALAudio: AL_OUT_OF_MEMORY in %s"), Text);
+						break;
+					case AL_INVALID_ENUM:
+						UE_LOG(LogALAudio, Warning, TEXT("ALAudio: AL_INVALID_ENUM in %s"), Text);
+						break;
+					case AL_INVALID_OPERATION:
+						UE_LOG(LogALAudio, Warning, TEXT("ALAudio: AL_INVALID_OPERATION in %s"), Text);
+						break;
+					default:
+						UE_LOG(LogALAudio, Warning, TEXT("ALAudio: Unknown Error NUM in %s"), Text);
+						break;
+				}
+			}
+		}
+		while( ( Error = alGetError() ) != AL_NO_ERROR );
+
+		return true;
+	}
+
+	return false;
+}
+
+FSoundSource* FALAudioDevice::CreateSoundSource()
+{
+	return new FALSoundSource(this);
+}
+
+bool FALAudioDevice::HasCompressedAudioInfoClass(USoundWave* SoundWave)
+{
+#if WITH_OGGVORBIS
+	return true;
+#else
+	return false;
+#endif
+}
+
+class ICompressedAudioInfo* FALAudioDevice::CreateCompressedAudioInfo(USoundWave* SoundWave)
+{
+#if WITH_OGGVORBIS
+	return new FVorbisAudioInfo();
+#else
+	return NULL;
+#endif
+}
diff --git a/Engine/Source/Runtime/ALAudio/Private/ALAudioSource.cpp b/Engine/Source/Runtime/ALAudio/Private/ALAudioSource.cpp
new file mode 100644
--- /dev/null
+++ b/Engine/Source/Runtime/ALAudio/Private/ALAudioSource.cpp
@@ -0,0 +1,293 @@
+// Copyright 1998-2015 Epic Games, Inc. All Rights Reserved.
+
+/*------------------------------------------------------------------------------------
+	FALSoundSource.
+------------------------------------------------------------------------------------*/
+
+#include "ALAudioDevice.h"
+#include "Engine.h"
+/**
+ * Initializes a source with a given wave instance and prepares it for playback.
+ *
+ * @param	WaveInstance	wave instace being primed for playback
+ * @return	TRUE if initialization was successful, FALSE otherwise
+ */
+bool FALSoundSource::Init( FWaveInstance* InWaveInstance )
+{
+	check(InWaveInstance);
+	GetALDevice()->MakeCurrent(TEXT("FALSoundSource::Init()"));
+
+	if (InWaveInstance->OutputTarget != EAudioOutputTarget::Controller)
+	{
+		// Find matching buffer.
+		Buffer = FALSoundBuffer::Init( (FALAudioDevice*)AudioDevice, InWaveInstance->WaveData  );
+		if( Buffer )
+		{
+			SCOPE_CYCLE_COUNTER( STAT_AudioSourceInitTime );
+
+			WaveInstance = InWaveInstance;
+
+			// Enable/disable spatialization of sounds
+			alSourcei( SourceId, AL_SOURCE_RELATIVE, WaveInstance->bUseSpatialization ? AL_FALSE : AL_TRUE );
+
+			// Setting looping on a real time decompressed source suppresses the buffers processed message
+			alSourcei( SourceId, AL_LOOPING, ( WaveInstance->LoopingMode == LOOP_Forever ) ? AL_TRUE : AL_FALSE );
+
+			// Always queue up the first buffer
+			alSourceQueueBuffers(SourceId, 1, &Buffer->BufferId);
+			if( WaveInstance->LoopingMode == LOOP_WithNotification )
+			{
+				// We queue the sound twice for wave instances that use seamless looping so we can have smooth
+				// loop transitions. The downside is that we might play at most one frame worth of audio from the
+				// beginning of the wave after the wave stops looping.
+				alSourceQueueBuffers(SourceId, 1, &Buffer->BufferId);
+			}
+
+			Update();
+
+			// Initialization was successful.
+			return true;
+		}
+	}
+
+	// Failed to initialize source.
+	UE_LOG(LogALAudio, Warning, TEXT("Failed to initialize sound source with WaveInstance '%s'."), *InWaveInstance->WaveData->GetName());
+	UE_LOG(LogALAudio, Warning, TEXT("  SampleRate %d"), InWaveInstance->WaveData->SampleRate);
+	UE_LOG(LogALAudio, Warning, TEXT("  Channels %d"), InWaveInstance->WaveData->NumChannels);
+	return false;
+}
+
+/**
+ * Clean up any hardware referenced by the sound source
+ */
+FALSoundSource::~FALSoundSource( void )
+{
+	GetALDevice()->MakeCurrent(TEXT("FALSoundSource::~FALSoundSource()"));
+
+	// @todo openal: What do we do here
+	/// AudioDevice->DestroyEffect( this );
+
+	alDeleteSources( 1, &SourceId );
+}
+
+/**
+ * Updates the source specific parameter like e.g. volume and pitch based on the associated
+ * wave instance.
+ */
+void FALSoundSource::Update( void )
+{
+	SCOPE_CYCLE_COUNTER( STAT_AudioUpdateSources );
+
+	if( !WaveInstance || Paused )
+	{
+		return;
+	}
+
+	GetALDevice()->MakeCurrent(TEXT("FALSoundSource::Update()"));
+
+	float Volume = 1.0f;
+
+	if (AudioDevice->bIsDeviceMuted)
+	{
+		Volume = 0.0f;
+	}
+	else
+	{
+		Volume = WaveInstance->Volume * WaveInstance->VolumeMultiplier;
+		if (SetStereoBleed())
+		{
+			// Emulate the bleed to rear speakers followed by stereo fold down
+			Volume *= 1.25f;
+		}
+		Volume *= FApp::GetVolumeMultiplier();
+		Volume *= AudioDevice->PlatformAudioHeadroom;
+		Volume = FMath::Clamp(Volume, 0.0f, MAX_VOLUME);
+	}
+
+	float	Pitch	=	FMath::Clamp(WaveInstance->Pitch, MIN_PITCH, MAX_PITCH );
+
+	// Set whether to apply reverb
+	SetReverbApplied( true );
+
+	// Set the HighFrequencyGain value
+	SetHighFrequencyGain();
+
+	FVector Location;
+	FVector	Velocity;
+
+	// See file header for coordinate system explanation.
+	Location.X = WaveInstance->Location.X;
+	Location.Y = WaveInstance->Location.Z; // Z/Y swapped on purpose, see file header
+	Location.Z = WaveInstance->Location.Y; // Z/Y swapped on purpose, see file header
+
+	Velocity.X = WaveInstance->Velocity.X;
+	Velocity.Y = WaveInstance->Velocity.Z; // Z/Y swapped on purpose, see file header
+	Velocity.Z = WaveInstance->Velocity.Y; // Z/Y swapped on purpose, see file header
+
+	// Convert to meters.
+	Location *= AUDIO_DISTANCE_FACTOR;
+	Velocity *= AUDIO_DISTANCE_FACTOR;
+
+	// We're using a relative coordinate system for un- spatialized sounds.
+	FVector RelativeDirection = FVector::ZeroVector;
+	if (WaveInstance->bUseSpatialization)
+	{
+		FVector UnnormalizedDirection = ((FALAudioDevice*)AudioDevice)->InverseTransform.TransformPosition(WaveInstance->Location);
+		RelativeDirection = UnnormalizedDirection.GetSafeNormal();
+	}
+	else
+	{
+		Location = FVector(0.f, 0.f, 0.f);
+	}
+
+	FVector EmitterPosition;
+	EmitterPosition.X = RelativeDirection.Y;
+	EmitterPosition.Y = RelativeDirection.X;
+	EmitterPosition.Z = -RelativeDirection.Z;
+
+	alSourcef(SourceId, AL_GAIN, Volume);
+	alSourcef(SourceId, AL_PITCH, Pitch);
+
+	alSourcefv(SourceId, AL_POSITION, (ALfloat*)&EmitterPosition);
+	alSourcefv(SourceId, AL_VELOCITY, (ALfloat*)&Velocity);
+}
+
+/**
+ * Plays the current wave instance.
+ */
+void FALSoundSource::Play( void )
+{
+	if( WaveInstance )
+	{
+		GetALDevice()->MakeCurrent(TEXT("FALSoundSource::Play()"));
+
+		alSourcePlay( SourceId );
+		Paused = false;
+		Playing = true;
+	}
+}
+
+/**
+ * Stops the current wave instance and detaches it from the source.
+ */
+void FALSoundSource::Stop( void )
+{
+	if( WaveInstance )
+	{
+		GetALDevice()->MakeCurrent(TEXT("FALSoundSource::Stop()"));
+
+		alSourceStop( SourceId );
+		// This clears out any pending buffers that may or may not be queued or played
+		alSourcei( SourceId, AL_BUFFER, 0 );
+		Paused = false;
+		Playing = false;
+		Buffer = NULL;
+	}
+
+	FSoundSource::Stop();
+}
+
+/**
+ * Pauses playback of current wave instance.
+ */
+void FALSoundSource::Pause( void )
+{
+	if( WaveInstance )
+	{
+		GetALDevice()->MakeCurrent(TEXT("FALSoundSource::Pause()"));
+
+		alSourcePause( SourceId );
+		Paused = true;
+	}
+}
+
+/**
+ * Returns TRUE if an OpenAL source has finished playing
+ */
+bool FALSoundSource::IsSourceFinished( void )
+{
+	ALint State = AL_STOPPED;
+
+	GetALDevice()->MakeCurrent(TEXT("FALSoundSource::IsSourceFinished()"));
+
+	// Check the source for data to continue playing
+	alGetSourcei( SourceId, AL_SOURCE_STATE, &State );
+	if( State == AL_PLAYING || State == AL_PAUSED )
+	{
+		return false;
+	}
+
+	return true;
+}
+
+/**
+ * Handle dequeuing and requeuing of a single buffer
+ */
+void FALSoundSource::HandleQueuedBuffer( void )
+{
+	GetALDevice()->MakeCurrent(TEXT("FALSoundSource::HandleQueuedBuffer()"));
+
+	ALuint	DequeuedBuffer;
+
+	// Unqueue the processed buffers
+	alSourceUnqueueBuffers( SourceId, 1, &DequeuedBuffer );
+
+	// Notify the wave instance that the current (native) buffer has finished playing.
+	WaveInstance->NotifyFinished();
+
+	// Queue the same packet again for looping
+	alSourceQueueBuffers(SourceId, 1, &Buffer->BufferId);
+}
+
+/**
+ * Queries the status of the currently associated wave instance.
+ *
+ * @return	TRUE if the wave instance/ source has finished playback and FALSE if it is
+ *			currently playing or paused.
+ */
+bool FALSoundSource::IsFinished( void )
+{
+	GetALDevice()->MakeCurrent(TEXT("FALSoundSource::IsFinished()"));
+
+	if( WaveInstance )
+	{
+		// Check for a non starved, stopped source
+		if( IsSourceFinished() )
+		{
+			// Notify the wave instance that it has finished playing.
+			WaveInstance->NotifyFinished();
+			return( true );
+		}
+		else
+		{
+			// Check to see if any complete buffers have been processed
+			ALint BuffersProcessed;
+			alGetSourcei( SourceId, AL_BUFFERS_PROCESSED, &BuffersProcessed );
+
+			switch( BuffersProcessed )
+			{
+			case 0:
+				// No buffers need updating
+				break;
+
+			case 1:
+				// Standard case of 1 buffer expired which needs repopulating
+				HandleQueuedBuffer();
+				break;
+
+			case 2:
+				// Starvation case when the source has stopped
+				HandleQueuedBuffer();
+				HandleQueuedBuffer();
+
+				// Restart the source
+				alSourcePlay( SourceId );
+				break;
+			}
+		}
+
+		return false;
+	}
+
+	return true;
+}
diff --git a/Engine/Source/Runtime/ALAudio/Public/ALAudioDevice.h b/Engine/Source/Runtime/ALAudio/Public/ALAudioDevice.h
new file mode 100644
--- /dev/null
+++ b/Engine/Source/Runtime/ALAudio/Public/ALAudioDevice.h
@@ -0,0 +1,259 @@
+// Copyright 1998-2015 Epic Games, Inc. All Rights Reserved.
+
+/*=============================================================================
+	ALAudioDevice.h: Unreal OpenAL audio interface object.
+=============================================================================*/
+
+#pragma  once
+
+/*------------------------------------------------------------------------------------
+	Dependencies, helpers & forward declarations.
+------------------------------------------------------------------------------------*/
+
+class FALAudioDevice;
+
+#include "Engine.h"
+#include "SoundDefinitions.h"
+#include "AudioEffect.h"
+
+#if SUPPORTS_PRAGMA_PACK
+#pragma pack (push,8)
+#endif
+#define AL_NO_PROTOTYPES 1
+#define ALC_NO_PROTOTYPES 1
+#include "AL/al.h"
+#include "AL/alc.h"
+#if SUPPORTS_PRAGMA_PACK
+#pragma pack (pop)
+#endif
+
+DECLARE_LOG_CATEGORY_EXTERN(LogALAudio, Log, All);
+
+//     2 UU == 1"
+// <=> 1 UU == 0.0127 m
+#define AUDIO_DISTANCE_FACTOR ( 0.0127f )
+
+/**
+ * OpenAL implementation of FSoundBuffer, containing the wave data and format information.
+ */
+class FALSoundBuffer : public FSoundBuffer
+{
+public:
+	/**
+	 * Constructor
+	 *
+	 * @param AudioDevice	audio device this sound buffer is going to be attached to.
+	 */
+	FALSoundBuffer( FALAudioDevice* AudioDevice );
+
+	/**
+	 * Destructor
+	 *
+	 * Frees wave data and detaches itself from audio device.
+	 */
+	~FALSoundBuffer( void );
+
+	/**
+	 * Static function used to create a buffer.
+	 *
+	 * @param	InWave				USoundWave to use as template and wave source
+	 * @param	AudioDevice			Audio device to attach created buffer to
+	 * @return	FALSoundBuffer pointer if buffer creation succeeded, NULL otherwise
+	 */
+	static FALSoundBuffer* Init(  FALAudioDevice* ,USoundWave* );
+
+	/**
+	 * Static function used to create a IOSAudio buffer and upload decompressed ogg vorbis data to.
+	 *
+	 * @param Wave			USoundWave to use as template and wave source
+	 * @param AudioDevice	audio device to attach created buffer to
+	 * @return FIOSAudioSoundBuffer pointer if buffer creation succeeded, NULL otherwise
+	 */
+	static void CreateNativeBuffer(FALAudioDevice* AudioDevice, USoundWave* InWave, FALSoundBuffer* &Buffer);
+
+	/**
+	 * Returns the size of this buffer in bytes.
+	 *
+	 * @return Size in bytes
+	 */
+	int GetSize( void )
+	{
+		return( BufferSize );
+	}
+
+	/** Buffer id used to reference the data stored in AL. */
+	ALuint					BufferId;
+	/** Format of the data internal to OpenAL */
+	ALuint					InternalFormat;
+
+	/** Number of bytes stored in OpenAL, or the size of the ogg vorbis data */
+	int						BufferSize;
+	/** Sample rate of the ogg vorbis data - typically 44100 or 22050 */
+	int						SampleRate;
+};
+
+/**
+ * OpenAL implementation of FSoundSource, the interface used to play, stop and update sources
+ */
+class FALSoundSource : public FSoundSource
+{
+public:
+	/**
+	 * Constructor
+	 *
+	 * @param	InAudioDevice	audio device this source is attached to
+	 */
+	FALSoundSource( class FAudioDevice* InAudioDevice )
+	:	FSoundSource( InAudioDevice ),
+		SourceId( 0 ),
+		Buffer( NULL )
+	{}
+
+	/**
+	 * Destructor
+	 */
+	~FALSoundSource( void );
+
+	/**
+	 * Initializes a source with a given wave instance and prepares it for playback.
+	 *
+	 * @param	WaveInstance	wave instance being primed for playback
+	 * @return	TRUE			if initialization was successful, FALSE otherwise
+	 */
+	virtual bool Init( FWaveInstance* WaveInstance );
+
+	/**
+	 * Updates the source specific parameter like e.g. volume and pitch based on the associated
+	 * wave instance.
+	 */
+	virtual void Update( void );
+
+	/**
+	 * Plays the current wave instance.
+	 */
+	virtual void Play( void );
+
+	/**
+	 * Stops the current wave instance and detaches it from the source.
+	 */
+	virtual void Stop( void );
+
+	/**
+	 * Pauses playback of current wave instance.
+	 */
+	virtual void Pause( void );
+
+	/**
+	 * Queries the status of the currently associated wave instance.
+	 *
+	 * @return	TRUE if the wave instance/ source has finished playback and FALSE if it is
+	 *			currently playing or paused.
+	 */
+	virtual bool IsFinished( void );
+
+	/**
+	 * Access function for the source id
+	 */
+	ALuint GetSourceId( void ) const { return( SourceId ); }
+
+	/**
+	 * Returns TRUE if an OpenAL source has finished playing
+	 */
+	bool IsSourceFinished( void );
+
+	/**
+	 * Handle dequeuing and requeuing of a single buffer
+	 */
+	void HandleQueuedBuffer( void );
+
+protected:
+	/** OpenAL source voice associated with this source/ channel. */
+	ALuint				SourceId;
+	/** Cached sound buffer associated with currently bound wave instance. */
+	FALSoundBuffer*		Buffer;
+
+	friend class FALAudioDevice;
+
+	/** Convenience function to avoid casting in multiple places */
+	FORCEINLINE FALAudioDevice * GetALDevice();
+};
+
+/**
+ * OpenAL implementation of an Unreal audio device.
+ */
+class FALAudioDevice : public FAudioDevice
+{
+public:
+	FALAudioDevice();
+
+	virtual ~FALAudioDevice() {}
+
+	virtual FName GetRuntimeFormat(USoundWave* SoundWave) override
+	{
+		static FName NAME_OGG(TEXT("OGG"));
+		return NAME_OGG;
+	}
+
+	/** Starts up any platform specific hardware/APIs */
+	virtual bool InitializeHardware() override;
+
+	/**
+	 * Tears down audio device by stopping all sounds, removing all buffers,  destroying all sources, ... Called by both Destroy and ShutdownAfterError
+	 * to perform the actual tear down.
+	 */
+	virtual void TeardownHardware() override;
+
+	/**
+	 * Update the audio device and calculates the cached inverse transform later
+	 * on used for spatialization.
+	 */
+	virtual void UpdateHardware() override;
+
+	virtual bool HasCompressedAudioInfoClass(USoundWave* SoundWave) override;
+
+	virtual class ICompressedAudioInfo* CreateCompressedAudioInfo(USoundWave* SoundWave) override;
+
+	// Error checking.
+	bool alError( const TCHAR* Text, bool Log = true );
+
+	/**
+	 * Makes sure correct context is set.
+	 *
+	 * @param CallSiteIdentifier identifies the call site for debugging/error reporting purposes
+	 */
+	void MakeCurrent(const TCHAR * CallSiteIdentifier = nullptr);
+
+protected:
+
+	virtual FSoundSource* CreateSoundSource() override;
+
+	/** Returns the enum for the internal format for playing a sound with this number of channels. */
+	ALuint GetInternalFormat( int NumChannels );
+
+
+	// Variables.
+
+	// AL specific
+
+	/** Device/context used to play back sounds (static so it can be initialized early) */
+	ALCdevice*									HardwareDevice;
+	ALCcontext*									SoundContext;
+	void*										DLLHandle;
+
+	/** Formats for multichannel sounds */
+	ALenum										Surround40Format;
+	ALenum										Surround51Format;
+	ALenum										Surround61Format;
+	ALenum										Surround71Format;
+
+	/** Inverse listener transformation, used for spatialization */
+	FMatrix										InverseTransform;
+
+	friend class FALSoundBuffer;
+	friend class FALSoundSource;
+};
+
+FORCEINLINE FALAudioDevice * FALSoundSource::GetALDevice()
+{
+	return static_cast<FALAudioDevice *>(AudioDevice);
+}
diff --git a/Engine/Source/Runtime/Linux/ALAudio/ALAudio.Build.cs b/Engine/Source/Runtime/Linux/ALAudio/ALAudio.Build.cs
deleted file mode 100644
diff --git a/Engine/Source/Runtime/Linux/ALAudio/Private/ALAudioBuffer.cpp b/Engine/Source/Runtime/Linux/ALAudio/Private/ALAudioBuffer.cpp
deleted file mode 100644
diff --git a/Engine/Source/Runtime/Linux/ALAudio/Private/ALAudioDevice.cpp b/Engine/Source/Runtime/Linux/ALAudio/Private/ALAudioDevice.cpp
deleted file mode 100644
diff --git a/Engine/Source/Runtime/Linux/ALAudio/Private/ALAudioSource.cpp b/Engine/Source/Runtime/Linux/ALAudio/Private/ALAudioSource.cpp
deleted file mode 100644
diff --git a/Engine/Source/Runtime/Linux/ALAudio/Public/ALAudioDevice.h b/Engine/Source/Runtime/Linux/ALAudio/Public/ALAudioDevice.h
deleted file mode 100644
