Fix CIS compile error in UE4Game HTML5.

diff --git a/Engine/Source/Runtime/ALAudio/Private/ALAudioBuffer.cpp b/Engine/Source/Runtime/ALAudio/Private/ALAudioBuffer.cpp
--- a/Engine/Source/Runtime/ALAudio/Private/ALAudioBuffer.cpp
+++ b/Engine/Source/Runtime/ALAudio/Private/ALAudioBuffer.cpp
@@ -145,7 +145,7 @@ void FALSoundBuffer::CreateNativeBuffer(FALAudioDevice* AudioDevice, USoundWave*
 	else
 	{
 		// get the raw data
-		uint8* SoundData = reinterpret_cast<uint8*>(Wave->RawData.Lock(LOCK_READ_ONLY));
+		const uint8* SoundData = reinterpret_cast<const uint8*>(Wave->RawData.Lock(LOCK_READ_ONLY));
 		// it's (possibly) a pointer to a wave file, so skip over the header
 		int SoundDataSize = Wave->RawData.GetBulkDataSize();
 
