Fix ALAudioDevice compile error due to update to MaxChannels being split out from MaxSources

diff --git a/Engine/Source/Runtime/ALAudio/Private/ALAudioDevice.cpp b/Engine/Source/Runtime/ALAudio/Private/ALAudioDevice.cpp
--- a/Engine/Source/Runtime/ALAudio/Private/ALAudioDevice.cpp
+++ b/Engine/Source/Runtime/ALAudio/Private/ALAudioDevice.cpp
@@ -112,11 +112,6 @@ bool FALAudioDevice::InitializeHardware( void )
 	Effects = NULL;
 	DLLHandle = NULL;
 
-	// Default to sensible channel count.
-	if( MaxChannels < 1 )
-	{
-		MaxChannels = 32;
-	}
 	// Load ogg and vorbis dlls if they haven't been loaded yet
 	//LoadVorbisLibraries();
 
@@ -173,16 +168,20 @@ bool FALAudioDevice::InitializeHardware( void )
 #endif
 	// Initialize channels.
 	alError( TEXT( "Emptying error stack" ), 0 );
-	for( int i = 0; i < (( MaxChannels >  MAX_AUDIOCHANNELS ) ? MAX_AUDIOCHANNELS : MaxChannels); i++ )
+
+	const int32 MaxSources = GetMaxSources();
+	check(MaxSources <= MAX_AUDIOCHANNELS);
+
+	for (int32 i = 0; i < MaxSources; i++)
 	{
 		ALuint SourceId;
-		alGenSources( 1, &SourceId );
-		if( !alError( TEXT( "Init (creating sources)" ), 0 ) )
+		alGenSources(1, &SourceId);
+		if (!alError(TEXT("Init (creating sources)"), 0))
 		{
-			FALSoundSource* Source = new FALSoundSource( this );
+			FALSoundSource* Source = new FALSoundSource(this);
 			Source->SourceId = SourceId;
-			Sources.Add( Source );
-			FreeSources.Add( Source );
+			Sources.Add(Source);
+			FreeSources.Add(Source);
 		}
 		else
 		{
@@ -190,15 +189,8 @@ bool FALAudioDevice::InitializeHardware( void )
 		}
 	}
 
-	if( Sources.Num() < 1 )
-	{
-		UE_LOG(LogALAudio, Warning, TEXT("ALAudio: couldn't allocate any sources"));
-		return false ;
-	}
-
 	// Update MaxChannels in case we couldn't create enough sources.
-	MaxChannels = Sources.Num();
-	UE_LOG(LogALAudio, Verbose, TEXT("ALAudioDevice: Allocated %i sources"), MaxChannels);
+	UE_LOG(LogALAudio, Verbose, TEXT("ALAudioDevice: Allocated %i sources"), MaxSources);
 
 	// Use our own distance model.
 	alDistanceModel( AL_NONE );
